rename NiagaraSimulationStages plugin to ChaosFluids

diff --git a/Engine/Plugins/Experimental/ChaosFluids/ChaosFluids.uplugin b/Engine/Plugins/Experimental/ChaosFluids/ChaosFluids.uplugin
new file mode 100644
--- /dev/null
+++ b/Engine/Plugins/Experimental/ChaosFluids/ChaosFluids.uplugin
@@ -0,0 +1,30 @@
+{
+	"FileVersion": 3,
+	"Version": 1,
+	"VersionName": "1.0",
+	"FriendlyName": "ChaosFluids",
+	"Description": "",
+	"Category": "Other",
+	"CreatedBy": "Epic Games, Inc.",
+	"CreatedByURL": "",
+	"DocsURL": "",
+	"MarketplaceURL": "",
+	"SupportURL": "",
+	"CanContainContent": true,
+	"IsBetaVersion": false,
+	"IsExperimentalVersion": true,
+	"Installed": false,
+	"Modules": [
+		{
+			"Name": "ChaosFluids",
+			"Type": "Runtime",
+			"LoadingPhase": "Default"
+		}
+	],    
+	"Plugins": [
+        {
+            "Name": "Niagara",			
+            "Enabled": true
+        }
+	]
+}
\ No newline at end of file
diff --git a/Engine/Plugins/Experimental/ChaosFluids/Config/BaseChaosFluids.ini b/Engine/Plugins/Experimental/ChaosFluids/Config/BaseChaosFluids.ini
new file mode 100644
--- /dev/null
+++ b/Engine/Plugins/Experimental/ChaosFluids/Config/BaseChaosFluids.ini
@@ -0,0 +1,2 @@
+[CoreRedirects]
++PackageRedirects=(OldName="/NiagaraSimulationStages/",NewName="/ChaosFluids/",MatchSubstring=true)
\ No newline at end of file
diff --git a/Engine/Plugins/Experimental/ChaosFluids/Shaders/ChaosFluids.ush b/Engine/Plugins/Experimental/ChaosFluids/Shaders/ChaosFluids.ush
new file mode 100644
--- /dev/null
+++ b/Engine/Plugins/Experimental/ChaosFluids/Shaders/ChaosFluids.ush
@@ -0,0 +1,35 @@
+// Copyright Epic Games, Inc. All Rights Reserved.
+
+/*=============================================================================
+NiagaraDataInterfacePhysicsField.ush 
+=============================================================================*/
+
+#pragma once
+
+/* -----------------------------------------------------------------
+ * Rendering code for simulation stages
+ * -----------------------------------------------------------------
+ */
+
+
+#define MARCH_SINGLE_RAY(RayStart, RayDir, UnitStepSize, LocalStepSize, NumSteps, Volume, RetV) \
+RetV = float4(0,0,0,0); \
+float Transmittance = 1; \
+int Steps = ceil(NumSteps); \
+if(Steps >= 1.0 ) Steps += 1.0; \
+[loop] \
+for (int i = 0; i < Steps && Transmittance > 1e-5; i++) \
+{ \
+float Extinction; \
+float Scattering; \
+float3 Luminance; \
+Volume.ComputeStep(RayStart, RayDir, Extinction, Scattering, Luminance); \
+Luminance *= Scattering; \
+float ClampedExtinction = max(Extinction, 1e-6); \
+float StepTransmittance = exp(-Extinction * LocalStepSize); \
+RetV.rgb += Transmittance * (Luminance - Luminance * StepTransmittance) / ClampedExtinction; \
+Transmittance *= StepTransmittance; \
+if(i == Steps-2) UnitStepSize, LocalStepSize *= frac(NumSteps); \
+RayStart += RayDir * UnitStepSize; \
+} \
+RetV.a = 1. - Transmittance; \
diff --git a/Engine/Plugins/Experimental/ChaosFluids/Source/ChaosFluids/ChaosFluids.Build.cs b/Engine/Plugins/Experimental/ChaosFluids/Source/ChaosFluids/ChaosFluids.Build.cs
new file mode 100644
--- /dev/null
+++ b/Engine/Plugins/Experimental/ChaosFluids/Source/ChaosFluids/ChaosFluids.Build.cs
@@ -0,0 +1,52 @@
+// Copyright Epic Games, Inc. All Rights Reserved.
+
+namespace UnrealBuildTool.Rules
+{
+	public class ChaosFluids : ModuleRules
+	{
+		public ChaosFluids(ReadOnlyTargetRules Target) : base(Target)
+		{
+			PublicIncludePaths.AddRange(
+				new string[] {
+					// ... add public include paths required here ...
+				}
+				);
+
+			PrivateIncludePaths.AddRange(
+				new string[] {
+                    "../../../../Source/Runtime/Engine/",
+                    "../../../../Source/Runtime/RHI/",
+					// ... add other private include paths required here ...
+				}
+				);
+            
+            PublicDependencyModuleNames.AddRange(
+				new string[]
+				{
+                    "Core",
+                    "CoreUObject",
+                    "Engine",
+                    "RHI",
+                    "Engine",
+                    "RenderCore",
+					// ... add other public dependencies that you statically link with here ...
+				}
+				);
+
+			PrivateDependencyModuleNames.AddRange(
+				new string[]
+                {
+                    "Projects",
+					// ... add private dependencies that you statically link with here ...
+				}
+				);
+
+			DynamicallyLoadedModuleNames.AddRange(
+				new string[]
+				{
+					// ... add any modules that your module loads dynamically here ...
+				}
+				);
+		}
+	}
+}
diff --git a/Engine/Plugins/Experimental/ChaosFluids/Source/ChaosFluids/Private/ChaosFluids.cpp b/Engine/Plugins/Experimental/ChaosFluids/Source/ChaosFluids/Private/ChaosFluids.cpp
new file mode 100644
--- /dev/null
+++ b/Engine/Plugins/Experimental/ChaosFluids/Source/ChaosFluids/Private/ChaosFluids.cpp
@@ -0,0 +1,32 @@
+// Copyright Epic Games, Inc. All Rights Reserved.
+
+#include "CoreMinimal.h"
+#include "Modules/ModuleManager.h"
+#include "IChaosFluids.h"
+#include "Interfaces/IPluginManager.h"
+#include "Misc/Paths.h"
+#include "ShaderCore.h"
+
+
+class FChaosFluids : public IChaosFluids
+{
+	/** IModuleInterface implementation */
+	virtual void StartupModule() override;
+	virtual void ShutdownModule() override;
+};
+
+IMPLEMENT_MODULE(FChaosFluids, ChaosFluids )
+
+
+void FChaosFluids::StartupModule()
+{
+	FString PluginShaderDir = FPaths::Combine(IPluginManager::Get().FindPlugin(TEXT("ChaosFluids"))->GetBaseDir(), TEXT("Shaders"));
+	AddShaderSourceDirectoryMapping(TEXT("/Plugin/Experimental/ChaosFluids"), PluginShaderDir);
+}
+
+
+void FChaosFluids::ShutdownModule()
+{
+	// This function may be called during shutdown to clean up your module.  For modules that support dynamic reloading,
+	// we call this function before unloading the module.
+}
diff --git a/Engine/Plugins/Experimental/ChaosFluids/Source/ChaosFluids/Public/IChaosFluids.h b/Engine/Plugins/Experimental/ChaosFluids/Source/ChaosFluids/Public/IChaosFluids.h
new file mode 100644
--- /dev/null
+++ b/Engine/Plugins/Experimental/ChaosFluids/Source/ChaosFluids/Public/IChaosFluids.h
@@ -0,0 +1,39 @@
+// Copyright Epic Games, Inc. All Rights Reserved.
+
+#pragma once
+
+#include "CoreMinimal.h"
+#include "Modules/ModuleInterface.h"
+#include "Modules/ModuleManager.h"
+
+
+/**
+ * The public interface to this module
+ */
+class IChaosFluids : public IModuleInterface
+{
+
+public:
+
+	/**
+	 * Singleton-like access to this module's interface.  This is just for convenience!
+	 * Beware of calling this during the shutdown phase, though.  Your module might have been unloaded already.
+	 *
+	 * @return Returns singleton instance, loading the module on demand if needed
+	 */
+	static inline IChaosFluids& Get()
+	{
+		return FModuleManager::LoadModuleChecked< IChaosFluids >( "ChaosFluids" );
+	}
+
+	/**
+	 * Checks to see if this module is loaded and ready.  It is only valid to call Get() if IsAvailable() returns true.
+	 *
+	 * @return True if the module is loaded and ready to use
+	 */
+	static inline bool IsAvailable()
+	{
+		return FModuleManager::Get().IsModuleLoaded( "ChaosFluids" );
+	}
+};
+
diff --git a/Engine/Plugins/FX/NiagaraSimulationStages/NiagaraSimulationStages.uplugin b/Engine/Plugins/FX/NiagaraSimulationStages/NiagaraSimulationStages.uplugin
deleted file mode 100644
diff --git a/Engine/Plugins/FX/NiagaraSimulationStages/Shaders/NiagaraSimulationStagesRendering.ush b/Engine/Plugins/FX/NiagaraSimulationStages/Shaders/NiagaraSimulationStagesRendering.ush
deleted file mode 100644
diff --git a/Engine/Plugins/FX/NiagaraSimulationStages/Source/NiagaraSimulationStages/NiagaraSimulationStages.Build.cs b/Engine/Plugins/FX/NiagaraSimulationStages/Source/NiagaraSimulationStages/NiagaraSimulationStages.Build.cs
deleted file mode 100644
diff --git a/Engine/Plugins/FX/NiagaraSimulationStages/Source/NiagaraSimulationStages/Private/NiagaraSimulationStages.cpp b/Engine/Plugins/FX/NiagaraSimulationStages/Source/NiagaraSimulationStages/Private/NiagaraSimulationStages.cpp
deleted file mode 100644
diff --git a/Engine/Plugins/FX/NiagaraSimulationStages/Source/NiagaraSimulationStages/Public/INiagaraSimulationStages.h b/Engine/Plugins/FX/NiagaraSimulationStages/Source/NiagaraSimulationStages/Public/INiagaraSimulationStages.h
deleted file mode 100644
